[INTEGRATE] Merging CL#2586385 using UE4-To-UE4-LauncherDev
Don't require NSApp boiler plate code in order to use SWebBrowser on Mac
Jira OPP-3500

diff --git a/Engine/Source/Programs/SlateViewer/Private/Mac/SlateViewerMainMac.cpp b/Engine/Source/Programs/SlateViewer/Private/Mac/SlateViewerMainMac.cpp
--- a/Engine/Source/Programs/SlateViewer/Private/Mac/SlateViewerMainMac.cpp
+++ b/Engine/Source/Programs/SlateViewer/Private/Mac/SlateViewerMainMac.cpp
@@ -3,9 +3,6 @@
 #include "SlateViewerApp.h"
 #include "ExceptionHandling.h"
 #include "CocoaThread.h"
-#if WITH_CEF3
-#import "include/cef_application_mac.h"
-#endif
 
 static FString GSavedCommandLine;
 
@@ -57,33 +54,6 @@ static FString GSavedCommandLine;
 
 @end
 
-#if WITH_CEF3
-@interface UE4Application : NSApplication<CefAppProtocol>
-{
-@private
-	BOOL bHandlingSendEvent;
-}
-@end
-
-@implementation UE4Application
-- (BOOL)isHandlingSendEvent
-{
-	return bHandlingSendEvent;
-}
-
-- (void)setHandlingSendEvent:(BOOL)handlingSendEvent
-{
-	bHandlingSendEvent = handlingSendEvent;
-}
-
-- (void)sendEvent:(NSEvent*)event
-{
-	CefScopedSendingEvent sendingEventScoper;
-	[super sendEvent:event];
-}
-@end
-#endif
-
 int main(int argc, char *argv[])
 {
 	for (int32 Option = 1; Option < argc; Option++)
@@ -108,11 +78,7 @@ int main(int argc, char *argv[])
 	}
 
 	SCOPED_AUTORELEASE_POOL;
-#if WITH_CEF3
-	[UE4Application sharedApplication];
-#else
 	[NSApplication sharedApplication];
-#endif
 	[NSApp setDelegate:[UE4AppDelegate new]];
 	[NSApp run];
 	return 0;
diff --git a/Engine/Source/Runtime/CEF3Utils/Private/Mac/MacCefApplicationCategory.cpp b/Engine/Source/Runtime/CEF3Utils/Private/Mac/MacCefApplicationCategory.cpp
new file mode 100644
--- /dev/null
+++ b/Engine/Source/Runtime/CEF3Utils/Private/Mac/MacCefApplicationCategory.cpp
@@ -0,0 +1,53 @@
+// Copyright 1998-2015 Epic Games, Inc. All Rights Reserved.
+
+#include "CEF3UtilsPrivatePCH.h"
+
+#if WITH_CEF3
+#import <objc/runtime.h>
+
+#pragma push_macro("OVERRIDE")
+#undef OVERRIDE // cef headers provide their own OVERRIDE macro
+#import "include/cef_application_mac.h"
+#pragma pop_macro("OVERRIDE")
+
+@interface NSApplication (CEF3UtilsApplication) < CefAppProtocol >
+{
+}
+-(void)cef3UtilsSendEvent : (NSEvent*)event;
+@end
+
+@implementation NSApplication (CEF3UtilsApplication)
+static BOOL bGHandlingSendEvent;
+
+// Since objc no longer supports class_poseAs, do black magic to be able to override sendEvent in NSApplication
++ (void)load
+{
+	static dispatch_once_t DispatchToken;
+	dispatch_once(&DispatchToken,
+	^{
+		Method OriginalMethod = class_getInstanceMethod([self class], @selector(sendEvent:));
+		Method PatchedMethod = class_getInstanceMethod([self class], @selector(cef3UtilsSendEvent:));
+
+		method_exchangeImplementations(OriginalMethod, PatchedMethod);
+	});
+}
+
+- (BOOL)isHandlingSendEvent
+{
+	NSNumber* Property = objc_getAssociatedObject(self, @selector(isHandlingSendEvent));
+	return [Property boolValue];
+}
+
+-(void)setHandlingSendEvent:(BOOL)handlingSendEvent
+{
+	objc_setAssociatedObject(self, @selector(isHandlingSendEvent), [NSNumber numberWithBool: handlingSendEvent], OBJC_ASSOCIATION_RETAIN_NONATOMIC);
+}
+
+// This will be swizzled with NSApplication -sendEvent:
+-(void)cef3UtilsSendEvent : (NSEvent*)event
+{
+	CefScopedSendingEvent sendingEventScoper;
+	[self cef3UtilsSendEvent : event];
+}
+@end
+#endif
