Deal with move failures for open cef3.log as warning instead of error

diff --git a/Engine/Source/Runtime/CEF3Utils/Private/CEF3Utils.cpp b/Engine/Source/Runtime/CEF3Utils/Private/CEF3Utils.cpp
--- a/Engine/Source/Runtime/CEF3Utils/Private/CEF3Utils.cpp
+++ b/Engine/Source/Runtime/CEF3Utils/Private/CEF3Utils.cpp
@@ -152,9 +152,10 @@ namespace CEF3Utils
 			FString(Cef3LogFile).Split(TEXT("."), &Name, &Extension, ESearchCase::CaseSensitive, ESearchDir::FromEnd);
 			FDateTime OriginalTime = FileManager.GetTimeStamp(*Cef3LogFile);
 			FString BackupFilename = FString::Printf(TEXT("%s%s%s.%s"), *Name, BACKUP_LOG_FILENAME_POSTFIX, *OriginalTime.ToString(), *Extension);
-			if (!FileManager.Move(*BackupFilename, *Cef3LogFile, false))
+			// do not retry resulting in an error if log still in use
+			if (!FileManager.Move(*BackupFilename, *Cef3LogFile, false, false, false, true))
 			{
-				UE_LOG(LogCEF3Utils, Error, TEXT("Failed to backup cef3.log"));
+				UE_LOG(LogCEF3Utils, Warning, TEXT("Failed to backup cef3.log"));
 			}
 		}
 	}
