WebBrowser: Log CEF DLL path and version during module load, and check against compiled version

Because libcef is loaded dynamically and can exist in multiple places on a machine, there's a chance we could be loading the wrong DLL but with a compatible API, resulting in subtle bugs or undefined behavior.

This is the first step in diagnosing hard-to-reproduce bugs like https://shithappens.epicgames.com/incidents/31275

#rnx

diff --git a/Engine/Source/Runtime/CEF3Utils/Private/CEF3Utils.cpp b/Engine/Source/Runtime/CEF3Utils/Private/CEF3Utils.cpp
--- a/Engine/Source/Runtime/CEF3Utils/Private/CEF3Utils.cpp
+++ b/Engine/Source/Runtime/CEF3Utils/Private/CEF3Utils.cpp
@@ -27,11 +27,12 @@ IMPLEMENT_MODULE(FDefaultModuleImpl, CEF3Utils);
 namespace CEF3Utils
 {
 #if PLATFORM_WINDOWS
-    void* CEF3DLLHandle = nullptr;
+	void* CEF3DLLHandle = nullptr;
 	void* ElfHandle = nullptr;
 	void* D3DHandle = nullptr;
 	void* GLESHandle = nullptr;
-    void* EGLHandle = nullptr;
+	void* EGLHandle = nullptr;
+	FString DllPath;
 #elif PLATFORM_MAC
 	// Dynamically load the CEF framework library.
 	CefScopedLibraryLoader *CEFLibraryLoader = nullptr;
@@ -59,9 +60,9 @@ namespace CEF3Utils
 	{
 #if PLATFORM_WINDOWS
 	#if PLATFORM_64BITS
-		FString DllPath(FPaths::Combine(*FPaths::EngineDir(), TEXT("Binaries/ThirdParty/CEF3/Win64")));
+		DllPath = FPaths::Combine(*FPaths::EngineDir(), TEXT("Binaries/ThirdParty/CEF3/Win64"));
 	#else
-		FString DllPath(FPaths::Combine(*FPaths::EngineDir(), TEXT("Binaries/ThirdParty/CEF3/Win32")));
+		DllPath = FPaths::Combine(*FPaths::EngineDir(), TEXT("Binaries/ThirdParty/CEF3/Win32"));
 	#endif
 
 		FPlatformProcess::PushDllDirectory(*DllPath);
@@ -74,6 +75,7 @@ namespace CEF3Utils
 			EGLHandle = LoadDllCEF(FPaths::Combine(*DllPath, TEXT("libEGL.dll")));
 		}
 		FPlatformProcess::PopDllDirectory(*DllPath);
+		DllPath = FPaths::ConvertRelativePathToFull(DllPath);
 		return CEF3DLLHandle != nullptr;
 #elif PLATFORM_MAC
 		// Dynamically load the CEF framework library.
@@ -140,12 +142,14 @@ namespace CEF3Utils
 	}
 #endif
 
-#if PLATFORM_MAC
 	FString GetCEF3ModulePath()
 	{
-		 return FrameworkPath;
-	}
+#if PLATFORM_WINDOWS
+		return DllPath;
+#elif PLATFORM_MAC
+		return FrameworkPath;
 #endif
+	}
 
 	void BackupCEF3Logfile(const FString& LogFilePath)
 	{
diff --git a/Engine/Source/Runtime/CEF3Utils/Public/CEF3Utils.h b/Engine/Source/Runtime/CEF3Utils/Public/CEF3Utils.h
--- a/Engine/Source/Runtime/CEF3Utils/Public/CEF3Utils.h
+++ b/Engine/Source/Runtime/CEF3Utils/Public/CEF3Utils.h
@@ -26,12 +26,10 @@ namespace CEF3Utils
 	CEF3UTILS_API void* GetCEF3ModuleHandle();
 #endif
 
-#if PLATFORM_MAC
 	/**
 	  * Get the module loaded path
 	 */
 	CEF3UTILS_API FString GetCEF3ModulePath();
-#endif
 
 	/**
 	 * Move the current cef3.log file to a backup file, so CEF makes a new log when it starts up.
diff --git a/Engine/Source/Runtime/WebBrowser/Private/WebBrowserModule.cpp b/Engine/Source/Runtime/WebBrowser/Private/WebBrowserModule.cpp
--- a/Engine/Source/Runtime/WebBrowser/Private/WebBrowserModule.cpp
+++ b/Engine/Source/Runtime/WebBrowser/Private/WebBrowserModule.cpp
@@ -8,6 +8,7 @@
 #include "Misc/Paths.h"
 #if WITH_CEF3
 #	include "CEF3Utils.h"
+#	include "include/cef_version.h"
 #	if PLATFORM_MAC
 #		include "include/wrapper/cef_library_loader.h"
 #	endif
@@ -61,8 +62,18 @@ void FWebBrowserModule::StartupModule()
 	if (!CEFLibraryLoader->LoadInMain(TCHAR_TO_ANSI(*CEF3Utils::GetCEF3ModulePath())))
 	{
 		UE_LOG(LogWebBrowser, Error, TEXT("Chromium loader initialization failed"));
+		return;
 	}
 #endif // PLATFORM_MAC
+	const int CefVersionMajor = cef_version_info(0);
+	const int CefVersionMinor = cef_version_info(1);
+	const int CefVersionPatch = cef_version_info(2);
+	const int CefCommitNumber = cef_version_info(3);
+	UE_LOG(LogWebBrowser, Log, TEXT("Loaded CEF3 version %i.%i.%i.%i from %s"), CefVersionMajor, CefVersionMinor, CefVersionPatch, CefCommitNumber, *CEF3Utils::GetCEF3ModulePath());
+	if (CefVersionMajor != CEF_VERSION_MAJOR || CefVersionMinor != CEF_VERSION_MINOR || CefVersionPatch != CEF_VERSION_PATCH || CefCommitNumber != CEF_COMMIT_NUMBER)
+	{
+		UE_LOG(LogWebBrowser, Warning, TEXT("CEF3 loaded version mismatch! Module was built against %i.%i.%i.%i, check if library loading path is correct"), CEF_VERSION_MAJOR, CEF_VERSION_MINOR, CEF_VERSION_PATCH, CEF_COMMIT_NUMBER);
+	}
 #endif
 }
 
