WebBrowser: Fix missing return value in CEF3Utils for Linux

#rnx

diff --git a/Engine/Source/Runtime/CEF3Utils/Private/CEF3Utils.cpp b/Engine/Source/Runtime/CEF3Utils/Private/CEF3Utils.cpp
--- a/Engine/Source/Runtime/CEF3Utils/Private/CEF3Utils.cpp
+++ b/Engine/Source/Runtime/CEF3Utils/Private/CEF3Utils.cpp
@@ -148,6 +148,9 @@ namespace CEF3Utils
 		return DllPath;
 #elif PLATFORM_MAC
 		return FrameworkPath;
+#elif PLATFORM_LINUX
+		return FString();
+#else
 #endif
 	}
 
