WebBrowser: Fix build on Mac when bCompileCEF3=True

#rnx
#Synced-CL 50704910

diff --git a/Engine/Source/Runtime/CEF3Utils/Private/Mac/MacCefApplicationCategory.cpp b/Engine/Source/Runtime/CEF3Utils/Private/Mac/MacCefApplicationCategory.cpp
--- a/Engine/Source/Runtime/CEF3Utils/Private/Mac/MacCefApplicationCategory.cpp
+++ b/Engine/Source/Runtime/CEF3Utils/Private/Mac/MacCefApplicationCategory.cpp
@@ -3,6 +3,8 @@
 #include "CoreMinimal.h"
 
 #if WITH_CEF3
+#include "Mac/MacSystemIncludes.h"
+
 #import <objc/runtime.h>
 
 #pragma push_macro("OVERRIDE")
